Add dependency on Json module

#Synced-CL 47063821

diff --git a/Engine/Source/Developer/SubmitToolCore/SubmitToolCore.Build.cs b/Engine/Source/Developer/SubmitToolCore/SubmitToolCore.Build.cs
--- a/Engine/Source/Developer/SubmitToolCore/SubmitToolCore.Build.cs
+++ b/Engine/Source/Developer/SubmitToolCore/SubmitToolCore.Build.cs
@@ -16,6 +16,7 @@ public class SubmitToolCore : ModuleRules
 				"AppFramework",
 				"Core",
 				"CoreUObject",
+				"Json",
 				"OutputLog"
 			}
 		);
