Adding build.version file as a source of last synced changelist

#Synced-CL 48450364

diff --git a/Engine/Source/Developer/SubmitToolCore/Private/Configuration/Configuration.cpp b/Engine/Source/Developer/SubmitToolCore/Private/Configuration/Configuration.cpp
--- a/Engine/Source/Developer/SubmitToolCore/Private/Configuration/Configuration.cpp
+++ b/Engine/Source/Developer/SubmitToolCore/Private/Configuration/Configuration.cpp
@@ -72,25 +72,8 @@ void FConfiguration::Init()
 	Configuration.Values.Add(TEXT("$(SubmitToolLogsDir)"), *FPaths::ConvertRelativePathToFull(FPaths::ProjectLogDir()));
 
 	Configuration.Values.Add(TEXT("$(SubmitToolConfigDir)"), *FPaths::ConvertRelativePathToFull(FPaths::ProjectConfigDir()));
-	
-	const FString UgsStateFile = RootDir / TEXT(".ugs/state.json");	
-	FString LastSyncedCL = TEXT("Unknown");
-
-	FString FileContents;
-	if (FPaths::FileExists(UgsStateFile) && FFileHelper::LoadFileToString(FileContents, *UgsStateFile))
-	{
-		const FString LastCLField = TEXT("LastSyncChangeNumber");
-		TSharedPtr<FJsonObject> RootJsonObject;
-		TSharedRef<TJsonReader<>> Reader = TJsonReaderFactory<>::Create(FileContents);
-		FJsonSerializer::Deserialize(Reader, RootJsonObject);
-		uint64 LastSyncedCLValue = 0;
-		if (RootJsonObject.IsValid() && RootJsonObject->TryGetNumberField(LastCLField, LastSyncedCLValue))
-		{
-			LastSyncedCL = LexToString(LastSyncedCLValue);
-		}
-	}
 
-	Configuration.Values.Add(TEXT("$(UGSSyncedCL)"), MoveTemp(LastSyncedCL));
+	Configuration.Values.Add(TEXT("$(UGSSyncedCL)"), GetLastSyncedCLFromUGS(RootDir));
 
 	Instance = MakeShared<FConfiguration>(Configuration);
 }
@@ -116,6 +99,30 @@ void FConfiguration::AddOrUpdateEntry(const FString& Key, const FString& NewValu
 	}
 }
 
+FString FConfiguration::GetLastSyncedCLFromUGS(const FString& InRootDir)
+{
+	auto ReadFileProperty = [](const FString& InFile, const FString& InProperty) {
+		FString FileContents;
+		if (FPaths::FileExists(InFile) && FFileHelper::LoadFileToString(FileContents, *InFile))
+		{
+			TSharedPtr<FJsonObject> RootJsonObject;
+			TSharedRef<TJsonReader<>> Reader = TJsonReaderFactory<>::Create(FileContents);
+			FJsonSerializer::Deserialize(Reader, RootJsonObject);
+			int64 PropertyValue;
+			if (RootJsonObject.IsValid() && RootJsonObject->TryGetNumberField(InProperty, PropertyValue))
+			{
+				return PropertyValue;
+			}
+		}
+		return -1LL;
+	};
+
+	int64 StateJsonCL = ReadFileProperty(InRootDir / TEXT(".ugs/state.json"), TEXT("LastSyncChangeNumber"));
+	int64 BuildVersionCL = ReadFileProperty(InRootDir / TEXT("Engine/Build/Build.version"), TEXT("Changelist"));
+
+	return LexToString(FMath::Max(StateJsonCL, BuildVersionCL));
+}
+
 FString FConfiguration::Substitute(const FString& InStr)
 {
 	if (!Instance.IsValid())
diff --git a/Engine/Source/Developer/SubmitToolCore/Public/Configuration/Configuration.h b/Engine/Source/Developer/SubmitToolCore/Public/Configuration/Configuration.h
--- a/Engine/Source/Developer/SubmitToolCore/Public/Configuration/Configuration.h
+++ b/Engine/Source/Developer/SubmitToolCore/Public/Configuration/Configuration.h
@@ -14,6 +14,7 @@ public:
 	static void AddOrUpdateEntry(const FString& Key, const FString& NewValue);
 
 private:
+	static FString GetLastSyncedCLFromUGS(const FString& InRootDir);
 	static TSharedPtr<FConfiguration> Instance;
 	TMap<FString, FString> Values;
 };
\ No newline at end of file
