Adding support for checking for filedisk presence when evaluating validators applicability to a CL

#Synced-CL 48717939

diff --git a/Engine/Source/Developer/SubmitToolCore/Public/Logic/Validators/ValidatorDefinition.h b/Engine/Source/Developer/SubmitToolCore/Public/Logic/Validators/ValidatorDefinition.h
--- a/Engine/Source/Developer/SubmitToolCore/Public/Logic/Validators/ValidatorDefinition.h
+++ b/Engine/Source/Developer/SubmitToolCore/Public/Logic/Validators/ValidatorDefinition.h
@@ -92,6 +92,12 @@ struct FValidatorDefinition
 	UPROPERTY()
 	FString AppliesToCLRegex;
 
+	/// 
+	/// Checks if all these files exists as an additional requirement for this validator to run
+	/// 
+	UPROPERTY()
+	TArray<FString> CheckFileExists;
+
 	///
 	/// Optional message to print when the validator doesn't apply to CL to give more context to AppliesToCLRegex
 	/// 
diff --git a/Engine/Source/Programs/SubmitTool/Private/Logic/Validators/ValidatorBase.cpp b/Engine/Source/Programs/SubmitTool/Private/Logic/Validators/ValidatorBase.cpp
--- a/Engine/Source/Programs/SubmitTool/Private/Logic/Validators/ValidatorBase.cpp
+++ b/Engine/Source/Programs/SubmitTool/Private/Logic/Validators/ValidatorBase.cpp
@@ -672,6 +672,14 @@ bool FValidatorBase::AppliesToCL(const FString& InCLDescription, const TArray<FS
 		return true;
 	}
 
+	for (const FString& Filepath : Definition->CheckFileExists)
+	{
+		if (!FPaths::FileExists(FConfiguration::Substitute(Filepath)))
+		{
+			return false;
+		}
+	}
+
 	for (const FSCFileRef& File : FilesInCL)
 	{
 		bool bIsIncrementalSkip;
