[SubmitTool] Fixing Mac Non-Unity build

#Synced-CL 48728512

diff --git a/Engine/Source/Developer/SubmitToolCore/Public/SubmitToolCoreUtils.h b/Engine/Source/Developer/SubmitToolCore/Public/SubmitToolCoreUtils.h
--- a/Engine/Source/Developer/SubmitToolCore/Public/SubmitToolCoreUtils.h
+++ b/Engine/Source/Developer/SubmitToolCore/Public/SubmitToolCoreUtils.h
@@ -3,6 +3,7 @@
 #pragma once
 
 #include "Containers/Array.h"
+#include "Containers/UnrealString.h"
 
 class FString;
 
