Add path normalization to source control file constructor

#Synced-CL 51168036

diff --git a/Engine/Source/Developer/SubmitToolCore/Public/Models/SCFile.h b/Engine/Source/Developer/SubmitToolCore/Public/Models/SCFile.h
--- a/Engine/Source/Developer/SubmitToolCore/Public/Models/SCFile.h
+++ b/Engine/Source/Developer/SubmitToolCore/Public/Models/SCFile.h
@@ -3,6 +3,7 @@
 #pragma once
 
 #include "Templates/SharedPointer.h"
+#include "Misc/Paths.h"
 #include "SCFile.generated.h"
 
 UENUM()
@@ -22,7 +23,10 @@ public:
 	: Filename(InFilename)
 	, DepotPath(InDepotPath)
 	, State(InState)
-	{ }
+	{ 
+		FPaths::NormalizeFilename(Filename);
+		FPaths::NormalizeFilename(DepotPath);
+	}
 	~FSCFile() {}
 
 	const FString& GetFilename() const { return Filename; }
